Fix broken ini and uproject setup
#rnx

diff --git a/Templates/TP_AEC_BlankBP/Config/DefaultEngine.ini b/Templates/TP_AEC_BlankBP/Config/DefaultEngine.ini
--- a/Templates/TP_AEC_BlankBP/Config/DefaultEngine.ini
+++ b/Templates/TP_AEC_BlankBP/Config/DefaultEngine.ini
@@ -17,7 +17,7 @@ r.LightPropagationVolume=0
 ; quality settings
 r.ReflectionCaptureResolution=2048
 r.AllowStaticLighting=True
-r.HighResScreenshotDelay 8
+r.HighResScreenshotDelay=8
 r.DefaultBackBufferPixelFormat=4
 r.AllowGlobalClipPlane=False
 
diff --git a/Templates/TP_AEC_BlankBP/TP_AEC_BlankBP.uproject b/Templates/TP_AEC_BlankBP/TP_AEC_BlankBP.uproject
--- a/Templates/TP_AEC_BlankBP/TP_AEC_BlankBP.uproject
+++ b/Templates/TP_AEC_BlankBP/TP_AEC_BlankBP.uproject
@@ -23,13 +23,13 @@
 		},
 		{
 			"Name": "VariantManager",
-			"Enabled": True
+			"Enabled": true
 		},
 		{
 			"Name": "PythonScriptPlugin",
 			"Enabled": true
 		},
-				{
+		{
 			"Name": "HDRIBackdrop",
 			"Enabled": true
 		},
diff --git a/Templates/TP_ME_BlankBP/Config/DefaultEngine.ini b/Templates/TP_ME_BlankBP/Config/DefaultEngine.ini
--- a/Templates/TP_ME_BlankBP/Config/DefaultEngine.ini
+++ b/Templates/TP_ME_BlankBP/Config/DefaultEngine.ini
@@ -17,7 +17,7 @@ r.LightPropagationVolume=0
 ; quality settings
 r.ReflectionCaptureResolution=2048
 r.AllowStaticLighting=True
-r.HighResScreenshotDelay 8
+r.HighResScreenshotDelay=8
 r.DefaultBackBufferPixelFormat=4
 r.AllowGlobalClipPlane=False
 
diff --git a/Templates/TP_MFG_BlankBP/Config/DefaultEngine.ini b/Templates/TP_MFG_BlankBP/Config/DefaultEngine.ini
--- a/Templates/TP_MFG_BlankBP/Config/DefaultEngine.ini
+++ b/Templates/TP_MFG_BlankBP/Config/DefaultEngine.ini
@@ -17,7 +17,7 @@ r.LightPropagationVolume=0
 ; quality settings
 r.ReflectionCaptureResolution=2048
 r.AllowStaticLighting=True
-r.HighResScreenshotDelay 8
+r.HighResScreenshotDelay=8
 r.DefaultBackBufferPixelFormat=4
 r.AllowGlobalClipPlane=False
 
diff --git a/Templates/TP_MFG_BlankBP/TP_MFG_BlankBP.uproject b/Templates/TP_MFG_BlankBP/TP_MFG_BlankBP.uproject
--- a/Templates/TP_MFG_BlankBP/TP_MFG_BlankBP.uproject
+++ b/Templates/TP_MFG_BlankBP/TP_MFG_BlankBP.uproject
@@ -23,13 +23,13 @@
 		},
 		{
 			"Name": "VariantManager",
-			"Enabled": True
+			"Enabled": true
 		},
 		{
 			"Name": "PythonScriptPlugin",
 			"Enabled": true
 		},
-				{
+		{
 			"Name": "HDRIBackdrop",
 			"Enabled": true
 		},
diff --git a/Templates/TP_PhotoStudioBP/Config/DefaultEngine.ini b/Templates/TP_PhotoStudioBP/Config/DefaultEngine.ini
--- a/Templates/TP_PhotoStudioBP/Config/DefaultEngine.ini
+++ b/Templates/TP_PhotoStudioBP/Config/DefaultEngine.ini
@@ -17,7 +17,7 @@ r.LightPropagationVolume=0
 ; quality settings
 r.ReflectionCaptureResolution=2048
 r.AllowStaticLighting=True
-r.HighResScreenshotDelay 8
+r.HighResScreenshotDelay=8
 r.DefaultBackBufferPixelFormat=4
 r.AllowGlobalClipPlane=False
 
diff --git a/Templates/TP_PhotoStudioBP/TP_PhotoStudioBP.uproject b/Templates/TP_PhotoStudioBP/TP_PhotoStudioBP.uproject
--- a/Templates/TP_PhotoStudioBP/TP_PhotoStudioBP.uproject
+++ b/Templates/TP_PhotoStudioBP/TP_PhotoStudioBP.uproject
@@ -23,7 +23,7 @@
 		},
 		{
 			"Name": "VariantManager",
-			"Enabled": True
+			"Enabled": true
 		},
 		{
 			"Name": "PythonScriptPlugin",
