TP_PhotoStudioBP Template Project.
- Removed content from developper folder
- Project Content addition

diff --git a/Templates/TP_PhotoStudioBP/Config/DefaultEngine.ini b/Templates/TP_PhotoStudioBP/Config/DefaultEngine.ini
--- a/Templates/TP_PhotoStudioBP/Config/DefaultEngine.ini
+++ b/Templates/TP_PhotoStudioBP/Config/DefaultEngine.ini
@@ -24,14 +24,14 @@ r.AllowGlobalClipPlane=False
 ; shader settings
 r.SupportSkyAtmosphereAffectsHeightFog=True
 r.DefaultFeature.LensFlare=True
-r.ClearCoatNormal=False
+r.ClearCoatNormal=True
 r.NormalMapsForStaticLighting=False
-
+r.GenerateMeshDistanceFields=True
 
 [/Script/EngineSettings.GameMapsSettings]
 ;defines what level to be loaded by default in the editor and at playback time
-GameDefaultMap=/Game/Main.Main
-EditorStartupMap=/Game/Main.Main
+GameDefaultMap=/Game/TP_PhotoStudioBP/Levels/Main/Main.Main
+EditorStartupMap=/Game/TP_PhotoStudioBP/Levels/Main/Main.Main
 
 
 [/Script/UnrealEd.UnrealEdEngine]
@@ -46,3 +46,51 @@ EditorStartupMap=/Game/Main.Main
 +TemplateMapInfos=(ThumbnailTexture=Texture2D'/Game/Developers/LevelTemplates/Thumbnails/Default.Default',Map="/Game/Developers/LevelTemplates/DarkBackGround")
 +TemplateMapInfos=(ThumbnailTexture=Texture2D'/Game/Developers/LevelTemplates/Thumbnails/Default.Default',Map="/Game/Developers/LevelTemplates/Exterior")
 +TemplateMapInfos=(ThumbnailTexture=Texture2D'/Game/Developers/LevelTemplates/Thumbnails/Default.Default',Map="/Game/Developers/LevelTemplates/TurnTable")
+
+[/Script/Engine.PhysicsSettings]
+DefaultGravityZ=-980.000000
+DefaultTerminalVelocity=4000.000000
+DefaultFluidFriction=0.300000
+SimulateScratchMemorySize=262144
+RagdollAggregateThreshold=4
+TriangleMeshTriangleMinAreaThreshold=5.000000
+bEnableShapeSharing=False
+bEnablePCM=True
+bEnableStabilization=False
+bWarnMissingLocks=True
+bEnable2DPhysics=False
+PhysicErrorCorrection=(PingExtrapolation=0.100000,PingLimit=100.000000,ErrorPerLinearDifference=1.000000,ErrorPerAngularDifference=1.000000,MaxRestoredStateError=1.000000,MaxLinearHardSnapDistance=400.000000,PositionLerp=0.000000,AngleLerp=0.400000,LinearVelocityCoefficient=100.000000,AngularVelocityCoefficient=10.000000,ErrorAccumulationSeconds=0.500000,ErrorAccumulationDistanceSq=15.000000,ErrorAccumulationSimilarity=100.000000)
+LockedAxis=Invalid
+DefaultDegreesOfFreedom=Full3D
+BounceThresholdVelocity=200.000000
+FrictionCombineMode=Average
+RestitutionCombineMode=Average
+MaxAngularVelocity=3600.000000
+MaxDepenetrationVelocity=0.000000
+ContactOffsetMultiplier=0.020000
+MinContactOffset=2.000000
+MaxContactOffset=8.000000
+bSimulateSkeletalMeshOnDedicatedServer=True
+DefaultShapeComplexity=CTF_UseSimpleAndComplex
+bDefaultHasComplexCollision=True
+bSuppressFaceRemapTable=False
+bSupportUVFromHitResults=False
+bDisableActiveActors=False
+bDisableKinematicStaticPairs=False
+bDisableKinematicKinematicPairs=False
+bDisableCCD=False
+bEnableEnhancedDeterminism=False
+AnimPhysicsMinDeltaTime=0.000000
+bSimulateAnimPhysicsAfterReset=False
+MaxPhysicsDeltaTime=0.033333
+bSubstepping=False
+bSubsteppingAsync=False
+MaxSubstepDeltaTime=0.016667
+MaxSubsteps=6
+SyncSceneSmoothingFactor=0.000000
+InitialAverageFrameRate=0.016667
+PhysXTreeRebuildRate=10
+DefaultBroadphaseSettings=(bUseMBPOnClient=False,bUseMBPOnServer=False,bUseMBPOuterBounds=False,MBPBounds=(Min=(X=0.000000,Y=0.000000,Z=0.000000),Max=(X=0.000000,Y=0.000000,Z=0.000000),IsValid=0),MBPOuterBounds=(Min=(X=0.000000,Y=0.000000,Z=0.000000),Max=(X=0.000000,Y=0.000000,Z=0.000000),IsValid=0),MBPNumSubdivs=2)
+ChaosSettings=(DefaultThreadingModel=TaskGraph,DedicatedThreadTickMode=VariableCappedWithTarget,DedicatedThreadBufferMode=Double)
+
+
diff --git a/Templates/TP_PhotoStudioBP/Config/DefaultGame.ini b/Templates/TP_PhotoStudioBP/Config/DefaultGame.ini
--- a/Templates/TP_PhotoStudioBP/Config/DefaultGame.ini
+++ b/Templates/TP_PhotoStudioBP/Config/DefaultGame.ini
@@ -1 +1,6 @@
 
+
+[/Script/EngineSettings.GeneralProjectSettings]
+ProjectName=TP_PhotoStudioBP
+Description=Automotive Photo Studio Templates
+
diff --git a/Templates/TP_PhotoStudioBP/TP_PhotoStudioBP.uproject b/Templates/TP_PhotoStudioBP/TP_PhotoStudioBP.uproject
--- a/Templates/TP_PhotoStudioBP/TP_PhotoStudioBP.uproject
+++ b/Templates/TP_PhotoStudioBP/TP_PhotoStudioBP.uproject
@@ -19,7 +19,10 @@
 		},
 		{
 			"Name": "DatasmithCADImporter",
-			"Enabled": true
+			"Enabled": true,
+			"SupportedTargetPlatforms": [
+				"Win64"
+			]
 		},
 		{
 			"Name": "VariantManager",
@@ -29,7 +32,7 @@
 			"Name": "PythonScriptPlugin",
 			"Enabled": true
 		},
-				{
+		{
 			"Name": "HDRIBackdrop",
 			"Enabled": true
 		},
@@ -37,7 +40,7 @@
 			"Name": "SequencerScripting",
 			"Enabled": true
 		},
-				{
+		{
 			"Name": "StaticMeshEditorExtension",
 			"Enabled": true
 		}
