Level renaming + associated folder structure.

diff --git a/Templates/TP_PhotoStudioBP/Config/DefaultEngine.ini b/Templates/TP_PhotoStudioBP/Config/DefaultEngine.ini
--- a/Templates/TP_PhotoStudioBP/Config/DefaultEngine.ini
+++ b/Templates/TP_PhotoStudioBP/Config/DefaultEngine.ini
@@ -30,11 +30,11 @@ r.GenerateMeshDistanceFields=True
 
 [/Script/EngineSettings.GameMapsSettings]
 ;defines what level to be loaded by default in the editor and at playback time
-GameDefaultMap=/Game/PhotoStudio/Levels/Main/Main.Main
-EditorStartupMap=/Game/PhotoStudio/Levels/Main/Main.Main
+GameDefaultMap=/Game/PhotoStudio/Levels/Cinematic/Cinematic.Cinematic
+EditorStartupMap=/Game/PhotoStudio/Levels/Cinematic/Cinematic.Cinematic
 
 [/Script/UnrealEd.UnrealEdEngine]
 ; Remove Engine Template maps
 -TemplateMapInfos=(ThumbnailTexture=Texture2D'/Engine/Maps/Templates/Thumbnails/Default.Default',Map="/Engine/Maps/Templates/Template_Default")
 +TemplateMapInfos=(ThumbnailTexture=Texture2D'/Engine/Maps/Templates/Thumbnails/TimeOfDay.TimeOfDay',Map="/Engine/Maps/Templates/TimeOfDay_Default")
--TemplateMapInfos=(ThumbnailTexture=Texture2D'/Engine/Maps/Templates/Thumbnails/VR-Basic.VR-Basic',Map="/Engine/Maps/Templates/VR-Basic")
\ No newline at end of file
+-TemplateMapInfos=(ThumbnailTexture=Texture2D'/Engine/Maps/Templates/Thumbnails/VR-Basic.VR-Basic',Map="/Engine/Maps/Templates/VR-Basic")
