- GameMode set in project, uses spectator pawn
- Sequence no longer loops
- PBR black on car glass and tires
- Added scene notes to cinematic scene
- Typo in folder name fixed
- Remove materials on mesh assets

diff --git a/Templates/TP_PhotoStudioBP/Config/DefaultEngine.ini b/Templates/TP_PhotoStudioBP/Config/DefaultEngine.ini
--- a/Templates/TP_PhotoStudioBP/Config/DefaultEngine.ini
+++ b/Templates/TP_PhotoStudioBP/Config/DefaultEngine.ini
@@ -32,6 +32,7 @@ r.GenerateMeshDistanceFields=True
 ;defines what level to be loaded by default in the editor and at playback time
 GameDefaultMap=/Game/PhotoStudio/Levels/Cinematic/Cinematic.Cinematic
 EditorStartupMap=/Game/PhotoStudio/Levels/Cinematic/Cinematic.Cinematic
+GlobalDefaultGameMode=/Game/PhotoStudio/Blueprints/BP_PhotoStageGameMode.BP_PhotoStageGameMode_C
 
 [/Script/UnrealEd.UnrealEdEngine]
 ; Remove Engine Template maps
