Photostudio template

Resizing images andicons to avoid aliasing

diff --git a/Templates/TP_PhotoStudioBP/Config/TemplateDefs.ini b/Templates/TP_PhotoStudioBP/Config/TemplateDefs.ini
--- a/Templates/TP_PhotoStudioBP/Config/TemplateDefs.ini
+++ b/Templates/TP_PhotoStudioBP/Config/TemplateDefs.ini
@@ -11,9 +11,9 @@ Categories=MFG
 LocalizedDisplayNames=(Language="en", Text="Photo Studio")
 
 ; The localized descriptions for this template, which is displayed when the template is selected.
-LocalizedDescriptions=(Language="en", Text="Photo Studio template to create presentations and cinematics.")
+LocalizedDescriptions=(Language="en", Text="This template enables you to create product presentations and cinematics in a photographic studio environment.\n\nThis template enables you to use Datasmith to import models from various design applications and easily light them.")
 ; The sort key to use for the template to bypass the default alphabetical sorting.
-; SortKey="_1"
+SortKey="1"
 
 ; A list of folders to avoid copying when creating a project from this template. 
 ; These should contain files only of concern to the template itself and not desired in the user's project.
