Adding Chinese translation to it #jira ueloc-4711 #fyi jamie.dale, michelle.zhao, alexis.zhang

diff --git a/Templates/TP_PhotoStudioBP/Config/TemplateDefs.ini b/Templates/TP_PhotoStudioBP/Config/TemplateDefs.ini
--- a/Templates/TP_PhotoStudioBP/Config/TemplateDefs.ini
+++ b/Templates/TP_PhotoStudioBP/Config/TemplateDefs.ini
@@ -16,6 +16,9 @@ LocalizedDescriptions=(Language="en", Text="This template enables you to create
 LocalizedDisplayNames=(Language="ko", Text="포토 스튜디오")
 LocalizedDescriptions=(Language="ko", Text="이 템플릿은 사진 스튜디오 환경에서 제품 시연과 시네마틱을 제작할 수 있도록 해줍니다.\n\n이 템플릿은 다양한 디자인 어플리케이션에서 데이터스미스를 이용한 모델 임포트와 시각화가 가능합니다.")
 
+LocalizedDisplayNames=(Language="zh-Hans", Text="摄影棚")
+LocalizedDescriptions=(Language="zh-Hans", Text="该模板可以让你在摄影棚环境中创建产品演示和动画。\n\n该模板允许使用Datasmith从多种设计软件中导入模型并轻松为其设置光照。")
+
 ; The sort key to use for the template to bypass the default alphabetical sorting.
 SortKey="1"
 
