Updating all templates to default to DX12 and support DX12 SM6 per our new project defaults.

diff --git a/Templates/TP_AEC_CollabBP/Config/DefaultEngine.ini b/Templates/TP_AEC_CollabBP/Config/DefaultEngine.ini
--- a/Templates/TP_AEC_CollabBP/Config/DefaultEngine.ini
+++ b/Templates/TP_AEC_CollabBP/Config/DefaultEngine.ini
@@ -157,3 +157,9 @@ MaxInternetClientRate=200000
 [/Script/AndroidRuntimeSettings.AndroidRuntimeSettings]
 bPublicLogFiles=True
 
+[/Script/WindowsTargetPlatform.WindowsTargetSettings]
+DefaultGraphicsRHI=DefaultGraphicsRHI_DX12
+-D3D12TargetedShaderFormats=PCD3D_SM5
++D3D12TargetedShaderFormats=PCD3D_SM6
+-D3D11TargetedShaderFormats=PCD3D_SM5
++D3D11TargetedShaderFormats=PCD3D_SM5
diff --git a/Templates/TP_AEC_HandheldARBP/Config/DefaultEngine.ini b/Templates/TP_AEC_HandheldARBP/Config/DefaultEngine.ini
--- a/Templates/TP_AEC_HandheldARBP/Config/DefaultEngine.ini
+++ b/Templates/TP_AEC_HandheldARBP/Config/DefaultEngine.ini
@@ -163,3 +163,10 @@ SyncSceneSmoothingFactor=0.000000
 AsyncSceneSmoothingFactor=0.990000
 InitialAverageFrameRate=0.016667
 PhysXTreeRebuildRate=10
+
+[/Script/WindowsTargetPlatform.WindowsTargetSettings]
+DefaultGraphicsRHI=DefaultGraphicsRHI_DX12
+-D3D12TargetedShaderFormats=PCD3D_SM5
++D3D12TargetedShaderFormats=PCD3D_SM6
+-D3D11TargetedShaderFormats=PCD3D_SM5
++D3D11TargetedShaderFormats=PCD3D_SM5
diff --git a/Templates/TP_AEC_ProdConfigBP/Config/DefaultEngine.ini b/Templates/TP_AEC_ProdConfigBP/Config/DefaultEngine.ini
--- a/Templates/TP_AEC_ProdConfigBP/Config/DefaultEngine.ini
+++ b/Templates/TP_AEC_ProdConfigBP/Config/DefaultEngine.ini
@@ -18,3 +18,9 @@ r.SupportSkyAtmosphereAffectsHeightFog=True
 EditorStartupMap=/Game/ProductAssets/Levels/Main.Main
 GameDefaultMap=/Game/ProductAssets/Levels/Main.Main
 
+[/Script/WindowsTargetPlatform.WindowsTargetSettings]
+DefaultGraphicsRHI=DefaultGraphicsRHI_DX12
+-D3D12TargetedShaderFormats=PCD3D_SM5
++D3D12TargetedShaderFormats=PCD3D_SM6
+-D3D11TargetedShaderFormats=PCD3D_SM5
++D3D11TargetedShaderFormats=PCD3D_SM5
diff --git a/Templates/TP_Blank/Config/DefaultEngine.ini b/Templates/TP_Blank/Config/DefaultEngine.ini
--- a/Templates/TP_Blank/Config/DefaultEngine.ini
+++ b/Templates/TP_Blank/Config/DefaultEngine.ini
@@ -4,7 +4,7 @@
 GameDefaultMap=/Engine/Maps/Templates/OpenWorld
 
 [/Script/WindowsTargetPlatform.WindowsTargetSettings]
-DefaultGraphicsRHI=DefaultGraphicsRHI_Default
+DefaultGraphicsRHI=DefaultGraphicsRHI_DX12
 -D3D12TargetedShaderFormats=PCD3D_SM5
 +D3D12TargetedShaderFormats=PCD3D_SM6
 -D3D11TargetedShaderFormats=PCD3D_SM5
diff --git a/Templates/TP_BlankBP/Config/DefaultEngine.ini b/Templates/TP_BlankBP/Config/DefaultEngine.ini
--- a/Templates/TP_BlankBP/Config/DefaultEngine.ini
+++ b/Templates/TP_BlankBP/Config/DefaultEngine.ini
@@ -4,7 +4,7 @@
 GameDefaultMap=/Engine/Maps/Templates/OpenWorld
 
 [/Script/WindowsTargetPlatform.WindowsTargetSettings]
-DefaultGraphicsRHI=DefaultGraphicsRHI_Default
+DefaultGraphicsRHI=DefaultGraphicsRHI_DX12
 -D3D12TargetedShaderFormats=PCD3D_SM5
 +D3D12TargetedShaderFormats=PCD3D_SM6
 -D3D11TargetedShaderFormats=PCD3D_SM5
diff --git a/Templates/TP_DMXBP/Config/DefaultEngine.ini b/Templates/TP_DMXBP/Config/DefaultEngine.ini
--- a/Templates/TP_DMXBP/Config/DefaultEngine.ini
+++ b/Templates/TP_DMXBP/Config/DefaultEngine.ini
@@ -28,3 +28,9 @@ Attributes=((Name="Color",Keywords="ColorWheel, Color1"),(Name="Red",Keywords="C
 bDefaultSendDMXEnabled=True
 bDefaultReceiveDMXEnabled=True
 
+[/Script/WindowsTargetPlatform.WindowsTargetSettings]
+DefaultGraphicsRHI=DefaultGraphicsRHI_DX12
+-D3D12TargetedShaderFormats=PCD3D_SM5
++D3D12TargetedShaderFormats=PCD3D_SM6
+-D3D11TargetedShaderFormats=PCD3D_SM5
++D3D11TargetedShaderFormats=PCD3D_SM5
diff --git a/Templates/TP_FirstPerson/Config/DefaultEngine.ini b/Templates/TP_FirstPerson/Config/DefaultEngine.ini
--- a/Templates/TP_FirstPerson/Config/DefaultEngine.ini
+++ b/Templates/TP_FirstPerson/Config/DefaultEngine.ini
@@ -25,7 +25,7 @@ r.DefaultFeature.AutoExposure.ExtendDefaultLuminanceRange=true
 
 
 [/Script/WindowsTargetPlatform.WindowsTargetSettings]
-DefaultGraphicsRHI=DefaultGraphicsRHI_Default
+DefaultGraphicsRHI=DefaultGraphicsRHI_DX12
 -D3D12TargetedShaderFormats=PCD3D_SM5
 +D3D12TargetedShaderFormats=PCD3D_SM6
 -D3D11TargetedShaderFormats=PCD3D_SM5
diff --git a/Templates/TP_FirstPersonBP/Config/DefaultEngine.ini b/Templates/TP_FirstPersonBP/Config/DefaultEngine.ini
--- a/Templates/TP_FirstPersonBP/Config/DefaultEngine.ini
+++ b/Templates/TP_FirstPersonBP/Config/DefaultEngine.ini
@@ -22,3 +22,10 @@ r.GenerateMeshDistanceFields=True
 r.DynamicGlobalIlluminationMethod=1
 r.Mobile.EnableNoPrecomputedLightingCSMShader=1
 r.DefaultFeature.AutoExposure.ExtendDefaultLuminanceRange=true
+
+[/Script/WindowsTargetPlatform.WindowsTargetSettings]
+DefaultGraphicsRHI=DefaultGraphicsRHI_DX12
+-D3D12TargetedShaderFormats=PCD3D_SM5
++D3D12TargetedShaderFormats=PCD3D_SM6
+-D3D11TargetedShaderFormats=PCD3D_SM5
++D3D11TargetedShaderFormats=PCD3D_SM5
diff --git a/Templates/TP_HandheldARBP/Config/DefaultEngine.ini b/Templates/TP_HandheldARBP/Config/DefaultEngine.ini
--- a/Templates/TP_HandheldARBP/Config/DefaultEngine.ini
+++ b/Templates/TP_HandheldARBP/Config/DefaultEngine.ini
@@ -161,3 +161,10 @@ SyncSceneSmoothingFactor=0.000000
 AsyncSceneSmoothingFactor=0.990000
 InitialAverageFrameRate=0.016667
 PhysXTreeRebuildRate=10
+
+[/Script/WindowsTargetPlatform.WindowsTargetSettings]
+DefaultGraphicsRHI=DefaultGraphicsRHI_DX12
+-D3D12TargetedShaderFormats=PCD3D_SM5
++D3D12TargetedShaderFormats=PCD3D_SM6
+-D3D11TargetedShaderFormats=PCD3D_SM5
++D3D11TargetedShaderFormats=PCD3D_SM5
diff --git a/Templates/TP_InCamVFXBP/Config/DefaultEngine.ini b/Templates/TP_InCamVFXBP/Config/DefaultEngine.ini
--- a/Templates/TP_InCamVFXBP/Config/DefaultEngine.ini
+++ b/Templates/TP_InCamVFXBP/Config/DefaultEngine.ini
@@ -20,10 +20,12 @@ r.DefaultBackBufferPixelFormat=4
 vr.AllowMotionBlurInVR=True
 
 [/Script/WindowsTargetPlatform.WindowsTargetSettings]
-Compiler=VisualStudio2019
--TargetedRHIs=PCD3D_SM5
-+TargetedRHIs=PCD3D_SM5
 DefaultGraphicsRHI=DefaultGraphicsRHI_DX12
+-D3D12TargetedShaderFormats=PCD3D_SM5
++D3D12TargetedShaderFormats=PCD3D_SM6
+-D3D11TargetedShaderFormats=PCD3D_SM5
++D3D11TargetedShaderFormats=PCD3D_SM5
+Compiler=VisualStudio2019
 bTarget32Bit=False
 AudioSampleRate=48000
 AudioCallbackBufferFrameSize=1024
diff --git a/Templates/TP_ME_BlankBP/Config/DefaultEngine.ini b/Templates/TP_ME_BlankBP/Config/DefaultEngine.ini
--- a/Templates/TP_ME_BlankBP/Config/DefaultEngine.ini
+++ b/Templates/TP_ME_BlankBP/Config/DefaultEngine.ini
@@ -91,3 +91,9 @@ EndpointSettings=(bEnableLogging=False,PurgeProcessedMessageDelaySeconds=30,Remo
 [/Script/DisplayClusterEditor.DisplayClusterEditorSettings]
 bEnabled=True
 
+[/Script/WindowsTargetPlatform.WindowsTargetSettings]
+DefaultGraphicsRHI=DefaultGraphicsRHI_DX12
+-D3D12TargetedShaderFormats=PCD3D_SM5
++D3D12TargetedShaderFormats=PCD3D_SM6
+-D3D11TargetedShaderFormats=PCD3D_SM5
++D3D11TargetedShaderFormats=PCD3D_SM5
diff --git a/Templates/TP_ME_VProdBP/Config/DefaultEngine.ini b/Templates/TP_ME_VProdBP/Config/DefaultEngine.ini
--- a/Templates/TP_ME_VProdBP/Config/DefaultEngine.ini
+++ b/Templates/TP_ME_VProdBP/Config/DefaultEngine.ini
@@ -91,3 +91,9 @@ EndpointSettings=(bEnableLogging=False,PurgeProcessedMessageDelaySeconds=30,Remo
 [/Script/DisplayClusterEditor.DisplayClusterEditorSettings]
 bEnabled=True
 
+[/Script/WindowsTargetPlatform.WindowsTargetSettings]
+DefaultGraphicsRHI=DefaultGraphicsRHI_DX12
+-D3D12TargetedShaderFormats=PCD3D_SM5
++D3D12TargetedShaderFormats=PCD3D_SM6
+-D3D11TargetedShaderFormats=PCD3D_SM5
++D3D11TargetedShaderFormats=PCD3D_SM5
diff --git a/Templates/TP_MFG_CollabBP/Config/DefaultEngine.ini b/Templates/TP_MFG_CollabBP/Config/DefaultEngine.ini
--- a/Templates/TP_MFG_CollabBP/Config/DefaultEngine.ini
+++ b/Templates/TP_MFG_CollabBP/Config/DefaultEngine.ini
@@ -148,3 +148,9 @@ ConfiguredLanSpeed=200000
 MaxClientRate=200000
 MaxInternetClientRate=200000
 
+[/Script/WindowsTargetPlatform.WindowsTargetSettings]
+DefaultGraphicsRHI=DefaultGraphicsRHI_DX12
+-D3D12TargetedShaderFormats=PCD3D_SM5
++D3D12TargetedShaderFormats=PCD3D_SM6
+-D3D11TargetedShaderFormats=PCD3D_SM5
++D3D11TargetedShaderFormats=PCD3D_SM5
diff --git a/Templates/TP_MFG_HandheldARBP/Config/DefaultEngine.ini b/Templates/TP_MFG_HandheldARBP/Config/DefaultEngine.ini
--- a/Templates/TP_MFG_HandheldARBP/Config/DefaultEngine.ini
+++ b/Templates/TP_MFG_HandheldARBP/Config/DefaultEngine.ini
@@ -160,3 +160,10 @@ SyncSceneSmoothingFactor=0.000000
 AsyncSceneSmoothingFactor=0.990000
 InitialAverageFrameRate=0.016667
 PhysXTreeRebuildRate=10
+
+[/Script/WindowsTargetPlatform.WindowsTargetSettings]
+DefaultGraphicsRHI=DefaultGraphicsRHI_DX12
+-D3D12TargetedShaderFormats=PCD3D_SM5
++D3D12TargetedShaderFormats=PCD3D_SM6
+-D3D11TargetedShaderFormats=PCD3D_SM5
++D3D11TargetedShaderFormats=PCD3D_SM5
diff --git a/Templates/TP_MFG_ProdConfigBP/Config/DefaultEngine.ini b/Templates/TP_MFG_ProdConfigBP/Config/DefaultEngine.ini
--- a/Templates/TP_MFG_ProdConfigBP/Config/DefaultEngine.ini
+++ b/Templates/TP_MFG_ProdConfigBP/Config/DefaultEngine.ini
@@ -19,3 +19,9 @@ r.GenerateMeshDistanceFields=True
 EditorStartupMap=/Game/ProductAssets/Levels/Main.Main
 GameDefaultMap=/Game/ProductAssets/Levels/Main.Main
 
+[/Script/WindowsTargetPlatform.WindowsTargetSettings]
+DefaultGraphicsRHI=DefaultGraphicsRHI_DX12
+-D3D12TargetedShaderFormats=PCD3D_SM5
++D3D12TargetedShaderFormats=PCD3D_SM6
+-D3D11TargetedShaderFormats=PCD3D_SM5
++D3D11TargetedShaderFormats=PCD3D_SM5
diff --git a/Templates/TP_PhotoStudioBP/Config/DefaultEngine.ini b/Templates/TP_PhotoStudioBP/Config/DefaultEngine.ini
--- a/Templates/TP_PhotoStudioBP/Config/DefaultEngine.ini
+++ b/Templates/TP_PhotoStudioBP/Config/DefaultEngine.ini
@@ -38,3 +38,10 @@ GlobalDefaultGameMode=/Game/PhotoStudio/Blueprints/BP_PhotoStageGameMode.BP_Phot
 -TemplateMapInfos=(ThumbnailTexture=Texture2D'/Engine/Maps/Templates/Thumbnails/Default.Default',Map="/Engine/Maps/Templates/Template_Default")
 +TemplateMapInfos=(ThumbnailTexture=Texture2D'/Engine/Maps/Templates/Thumbnails/TimeOfDay.TimeOfDay',Map="/Engine/Maps/Templates/TimeOfDay_Default")
 -TemplateMapInfos=(ThumbnailTexture=Texture2D'/Engine/Maps/Templates/Thumbnails/VR-Basic.VR-Basic',Map="/Engine/Maps/Templates/VR-Basic")
+
+[/Script/WindowsTargetPlatform.WindowsTargetSettings]
+DefaultGraphicsRHI=DefaultGraphicsRHI_DX12
+-D3D12TargetedShaderFormats=PCD3D_SM5
++D3D12TargetedShaderFormats=PCD3D_SM6
+-D3D11TargetedShaderFormats=PCD3D_SM5
++D3D11TargetedShaderFormats=PCD3D_SM5
diff --git a/Templates/TP_Puzzle/Config/DefaultEngine.ini b/Templates/TP_Puzzle/Config/DefaultEngine.ini
--- a/Templates/TP_Puzzle/Config/DefaultEngine.ini
+++ b/Templates/TP_Puzzle/Config/DefaultEngine.ini
@@ -10,4 +10,11 @@ GameInstanceClass=/Script/Engine.GameInstance
 GameDefaultMap=/Game/PuzzleCPP/Maps/PuzzleExampleMap
 ServerDefaultMap=/Engine/Maps/Entry
 GlobalDefaultGameMode=/Script/TP_Puzzle.TP_PuzzleGameMode
-GlobalDefaultServerGameMode=None
\ No newline at end of file
+GlobalDefaultServerGameMode=None
+
+[/Script/WindowsTargetPlatform.WindowsTargetSettings]
+DefaultGraphicsRHI=DefaultGraphicsRHI_DX12
+-D3D12TargetedShaderFormats=PCD3D_SM5
++D3D12TargetedShaderFormats=PCD3D_SM6
+-D3D11TargetedShaderFormats=PCD3D_SM5
++D3D11TargetedShaderFormats=PCD3D_SM5
diff --git a/Templates/TP_PuzzleBP/Config/DefaultEngine.ini b/Templates/TP_PuzzleBP/Config/DefaultEngine.ini
--- a/Templates/TP_PuzzleBP/Config/DefaultEngine.ini
+++ b/Templates/TP_PuzzleBP/Config/DefaultEngine.ini
@@ -9,4 +9,11 @@ bUseSplitscreen=True
 TwoPlayerSplitscreenLayout=Horizontal
 ThreePlayerSplitscreenLayout=FavorTop
 GlobalDefaultGameMode=/Game/PuzzleBP/Blueprints/PuzzleGameMode.PuzzleGameMode_C
-GlobalDefaultServerGameMode=None
\ No newline at end of file
+GlobalDefaultServerGameMode=None
+
+[/Script/WindowsTargetPlatform.WindowsTargetSettings]
+DefaultGraphicsRHI=DefaultGraphicsRHI_DX12
+-D3D12TargetedShaderFormats=PCD3D_SM5
++D3D12TargetedShaderFormats=PCD3D_SM6
+-D3D11TargetedShaderFormats=PCD3D_SM5
++D3D11TargetedShaderFormats=PCD3D_SM5
diff --git a/Templates/TP_SIM_Blank/Config/DefaultEngine.ini b/Templates/TP_SIM_Blank/Config/DefaultEngine.ini
--- a/Templates/TP_SIM_Blank/Config/DefaultEngine.ini
+++ b/Templates/TP_SIM_Blank/Config/DefaultEngine.ini
@@ -15,8 +15,10 @@ r.ReflectionMethod=1
 
 [/Script/WindowsTargetPlatform.WindowsTargetSettings]
 DefaultGraphicsRHI=DefaultGraphicsRHI_DX12
-D3D12TargetedShaderFormats=PCD3D_SM6
-D3D11TargetedShaderFormats=PCD3D_SM5
+-D3D12TargetedShaderFormats=PCD3D_SM5
++D3D12TargetedShaderFormats=PCD3D_SM6
+-D3D11TargetedShaderFormats=PCD3D_SM5
++D3D11TargetedShaderFormats=PCD3D_SM5
 Compiler=Default
 AudioSampleRate=48000
 AudioCallbackBufferFrameSize=1024
diff --git a/Templates/TP_SIM_BlankBP/Config/DefaultEngine.ini b/Templates/TP_SIM_BlankBP/Config/DefaultEngine.ini
--- a/Templates/TP_SIM_BlankBP/Config/DefaultEngine.ini
+++ b/Templates/TP_SIM_BlankBP/Config/DefaultEngine.ini
@@ -17,8 +17,10 @@ r.ReflectionMethod=1
 
 [/Script/WindowsTargetPlatform.WindowsTargetSettings]
 DefaultGraphicsRHI=DefaultGraphicsRHI_DX12
-D3D12TargetedShaderFormats=PCD3D_SM6
-D3D11TargetedShaderFormats=PCD3D_SM5
+-D3D12TargetedShaderFormats=PCD3D_SM5
++D3D12TargetedShaderFormats=PCD3D_SM6
+-D3D11TargetedShaderFormats=PCD3D_SM5
++D3D11TargetedShaderFormats=PCD3D_SM5
 Compiler=Default
 AudioSampleRate=48000
 AudioCallbackBufferFrameSize=1024
diff --git a/Templates/TP_SIM_Sandbox/Config/DefaultEngine.ini b/Templates/TP_SIM_Sandbox/Config/DefaultEngine.ini
--- a/Templates/TP_SIM_Sandbox/Config/DefaultEngine.ini
+++ b/Templates/TP_SIM_Sandbox/Config/DefaultEngine.ini
@@ -17,8 +17,10 @@ r.ReflectionMethod=1
 
 [/Script/WindowsTargetPlatform.WindowsTargetSettings]
 DefaultGraphicsRHI=DefaultGraphicsRHI_DX12
-D3D12TargetedShaderFormats=PCD3D_SM6
-D3D11TargetedShaderFormats=PCD3D_SM5
+-D3D12TargetedShaderFormats=PCD3D_SM5
++D3D12TargetedShaderFormats=PCD3D_SM6
+-D3D11TargetedShaderFormats=PCD3D_SM5
++D3D11TargetedShaderFormats=PCD3D_SM5
 Compiler=Default
 AudioSampleRate=48000
 AudioCallbackBufferFrameSize=1024
diff --git a/Templates/TP_SIM_SandboxBP/Config/DefaultEngine.ini b/Templates/TP_SIM_SandboxBP/Config/DefaultEngine.ini
--- a/Templates/TP_SIM_SandboxBP/Config/DefaultEngine.ini
+++ b/Templates/TP_SIM_SandboxBP/Config/DefaultEngine.ini
@@ -17,8 +17,10 @@ r.ReflectionMethod=1
 
 [/Script/WindowsTargetPlatform.WindowsTargetSettings]
 DefaultGraphicsRHI=DefaultGraphicsRHI_DX12
-D3D12TargetedShaderFormats=PCD3D_SM6
-D3D11TargetedShaderFormats=PCD3D_SM5
+-D3D12TargetedShaderFormats=PCD3D_SM5
++D3D12TargetedShaderFormats=PCD3D_SM6
+-D3D11TargetedShaderFormats=PCD3D_SM5
++D3D11TargetedShaderFormats=PCD3D_SM5
 Compiler=Default
 AudioSampleRate=48000
 AudioCallbackBufferFrameSize=1024
diff --git a/Templates/TP_VehicleAdv/Config/DefaultEngine.ini b/Templates/TP_VehicleAdv/Config/DefaultEngine.ini
--- a/Templates/TP_VehicleAdv/Config/DefaultEngine.ini
+++ b/Templates/TP_VehicleAdv/Config/DefaultEngine.ini
@@ -21,3 +21,9 @@ r.GenerateMeshDistanceFields=True
 r.DynamicGlobalIlluminationMethod=1
 r.Mobile.EnableNoPrecomputedLightingCSMShader=1
 
+[/Script/WindowsTargetPlatform.WindowsTargetSettings]
+DefaultGraphicsRHI=DefaultGraphicsRHI_DX12
+-D3D12TargetedShaderFormats=PCD3D_SM5
++D3D12TargetedShaderFormats=PCD3D_SM6
+-D3D11TargetedShaderFormats=PCD3D_SM5
++D3D11TargetedShaderFormats=PCD3D_SM5
diff --git a/Templates/TP_VehicleAdvBP/Config/DefaultEngine.ini b/Templates/TP_VehicleAdvBP/Config/DefaultEngine.ini
--- a/Templates/TP_VehicleAdvBP/Config/DefaultEngine.ini
+++ b/Templates/TP_VehicleAdvBP/Config/DefaultEngine.ini
@@ -27,7 +27,7 @@ r.ClearCoatNormal=True
 r.DefaultFeature.AutoExposure.ExtendDefaultLuminanceRange=true
 
 [/Script/WindowsTargetPlatform.WindowsTargetSettings]
-DefaultGraphicsRHI=DefaultGraphicsRHI_Default
+DefaultGraphicsRHI=DefaultGraphicsRHI_DX12
 -D3D12TargetedShaderFormats=PCD3D_SM5
 +D3D12TargetedShaderFormats=PCD3D_SM6
 -D3D11TargetedShaderFormats=PCD3D_SM5
diff --git a/Templates/TP_VirtualRealityBP/Config/DefaultEngine.ini b/Templates/TP_VirtualRealityBP/Config/DefaultEngine.ini
--- a/Templates/TP_VirtualRealityBP/Config/DefaultEngine.ini
+++ b/Templates/TP_VirtualRealityBP/Config/DefaultEngine.ini
@@ -137,6 +137,11 @@ GlobalDefaultServerGameMode=None
 bExplicitCanvasChildZOrder=True
 
 [/Script/WindowsTargetPlatform.WindowsTargetSettings]
+DefaultGraphicsRHI=DefaultGraphicsRHI_DX12
+-D3D12TargetedShaderFormats=PCD3D_SM5
++D3D12TargetedShaderFormats=PCD3D_SM6
+-D3D11TargetedShaderFormats=PCD3D_SM5
++D3D11TargetedShaderFormats=PCD3D_SM5
 Compiler=Default
 MinimumOSVersion=MSOS_Vista
 bTarget32Bit=False
diff --git a/Templates/TP_nDisplayBP/Config/DefaultEngine.ini b/Templates/TP_nDisplayBP/Config/DefaultEngine.ini
--- a/Templates/TP_nDisplayBP/Config/DefaultEngine.ini
+++ b/Templates/TP_nDisplayBP/Config/DefaultEngine.ini
@@ -20,3 +20,9 @@ vr.AllowMotionBlurInVR=True
 r.DefaultFeature.AutoExposure.Method=1
 r.DefaultFeature.AutoExposure.Bias=1.000000
 
+[/Script/WindowsTargetPlatform.WindowsTargetSettings]
+DefaultGraphicsRHI=DefaultGraphicsRHI_DX12
+-D3D12TargetedShaderFormats=PCD3D_SM5
++D3D12TargetedShaderFormats=PCD3D_SM6
+-D3D11TargetedShaderFormats=PCD3D_SM5
++D3D11TargetedShaderFormats=PCD3D_SM5
