[Template-PhotoStudio] Remove old files, Update DefaultEngine.ini with new rendering setting

diff --git a/Templates/TP_PhotoStudioBP/Config/DefaultEngine.ini b/Templates/TP_PhotoStudioBP/Config/DefaultEngine.ini
--- a/Templates/TP_PhotoStudioBP/Config/DefaultEngine.ini
+++ b/Templates/TP_PhotoStudioBP/Config/DefaultEngine.ini
@@ -1,47 +1,86 @@
+[Audio]
+UseAudioMixer=True
+
 [/Script/HardwareTargeting.HardwareTargetingSettings]
-; quality settings
 TargetedHardwareClass=Desktop
 AppliedTargetedHardwareClass=Desktop
 DefaultGraphicsPerformance=Maximum
 AppliedDefaultGraphicsPerformance=Maximum
 
+[/Script/WindowsTargetPlatform.WindowsTargetSettings]
+DefaultGraphicsRHI=DefaultGraphicsRHI_DX12
+DefaultGraphicsRHI=DefaultGraphicsRHI_DX12
+!D3D12TargetedShaderFormats=ClearArray
++D3D12TargetedShaderFormats=PCD3D_SM6
+
 [/Script/Engine.RendererSettings]
-; exposure control and tone mapping settings
-r.Tonemapper.Sharpen=2
+r.AllowStaticLighting=0
+r.GenerateMeshDistanceFields=True
+r.DynamicGlobalIlluminationMethod=1
+r.ReflectionMethod=1
+r.Shadow.Virtual.Enable=1
+r.Lumen.LightingDataFormat=1
+r.Lumen.Reflections.DenoiserTonemapRange =12
+r.VirtualTextures=True
 r.DefaultFeature.AutoExposure.ExtendDefaultLuminanceRange=True
+r.DefaultFeature.LocalExposure.HighlightContrastScale=0.8
+r.DefaultFeature.LocalExposure.ShadowContrastScale=0.8
+r.SkinCache.CompileShaders=True
+r.RayTracing=True
+r.Substrate=True
+r.VirtualTextures=True
+r.GBufferFormat=3
+r.Substrate.BytesPerPixel=240
+r.Substrate.EnableLayerSupport=True
+r.ReflectionCaptureResolution=1024
+r.RayTracing.Culling=0
+r.SkyAtmosphere.MultiScatteringLUT.HighQuality=1
+r.Translucency.StandardSeparated=1
+r.RayTracedTranslucency.MaxPrimaryHitEvents=3
+r.RayTracedTranslucency.ForceOpaque=0
+r.Lumen.HardwareRayTracing.HitLighting.ForceOpaque=0
+r.OIT.SortedPixels=False
+r.Material.EnergyConservation=True
+r.RayTracing.Shadows.Translucency=1
+r.RayTracing.NormalBias=0.18
+r.Lumen.HardwareRayTracing.LightingMode=2
+r.Translucent.UsesRectLights=True
+r.RayTracedTranslucency.Allowed=1
+r.RayTracedTranslucency.UseRayTracedRefraction=1
+r.Lumen.TranslucencyReflections.FrontLayer.EnableForProject=True
+r.Substrate.Debug.AdvancedVisualizationShaders=True
+r.PathTracing.VisibleLights=1
 
-;dynamic gi settings
-r.LightPropagationVolume=0
-
-; quality settings
-r.ReflectionCaptureResolution=2048
-r.AllowStaticLighting=True
-r.HighResScreenshotDelay=8
-r.DefaultBackBufferPixelFormat=4
-r.AllowGlobalClipPlane=False
-
-; shader settings
-r.SupportSkyAtmosphereAffectsHeightFog=True
-r.DefaultFeature.LensFlare=True
-r.ClearCoatNormal=True
-r.NormalMapsForStaticLighting=False
-r.GenerateMeshDistanceFields=True
+[/Script/Engine.VirtualTexturePoolConfig]
+DefaultSizeInMegabyte=64
+bPoolAutoGrowInEditor=False
++Pools=(Formats=(PF_DXT1),MinTileSize=136,MaxTileSize=136,SizeInMegabyte=78,bEnableResidencyMipMapBias=False,bAllowSizeScale=False,MinScaledSizeInMegabyte=0,MaxScaledSizeInMegabyte=0)
++Pools=(Formats=(PF_BC5),MinTileSize=136,MaxTileSize=136,SizeInMegabyte=84,bEnableResidencyMipMapBias=False,bAllowSizeScale=False,MinScaledSizeInMegabyte=0,MaxScaledSizeInMegabyte=0)
 
-[/Script/EngineSettings.GameMapsSettings]
-;defines what level to be loaded by default in the editor and at playback time
-GameDefaultMap=/Game/PhotoStudio/Levels/Cinematic/Cinematic.Cinematic
-EditorStartupMap=/Game/PhotoStudio/Levels/Cinematic/Cinematic.Cinematic
-GlobalDefaultGameMode=/Game/PhotoStudio/Blueprints/BP_PhotoStageGameMode.BP_PhotoStageGameMode_C
 
-[/Script/UnrealEd.UnrealEdEngine]
-; Remove Engine Template maps
--TemplateMapInfos=(ThumbnailTexture=Texture2D'/Engine/Maps/Templates/Thumbnails/Default.Default',Map="/Engine/Maps/Templates/Template_Default")
-+TemplateMapInfos=(ThumbnailTexture=Texture2D'/Engine/Maps/Templates/Thumbnails/TimeOfDay.TimeOfDay',Map="/Engine/Maps/Templates/TimeOfDay_Default")
--TemplateMapInfos=(ThumbnailTexture=Texture2D'/Engine/Maps/Templates/Thumbnails/VR-Basic.VR-Basic',Map="/Engine/Maps/Templates/VR-Basic")
+[/Script/WorldPartitionEditor.WorldPartitionEditorSettings]
+CommandletClass=Class'/Script/UnrealEd.WorldPartitionConvertCommandlet'
+
+[/Script/Engine.UserInterfaceSettings]
+bAuthorizeAutomaticWidgetVariableCreation=False
+FontDPIPreset=Standard
+FontDPI=72
+
+[/Script/AndroidFileServerEditor.AndroidFileServerRuntimeSettings]
+bEnablePlugin=True
+bAllowNetworkConnection=True
+SecurityToken=79BFBCEF42D211CDBF800C9F2597C868
+bIncludeInShipping=False
+bAllowExternalStartInShipping=False
+bCompileAFSProject=False
+bUseCompression=False
+bLogFiles=False
+bReportStats=False
+ConnectionType=USBOnly
+bUseManualIPAddress=False
+ManualIPAddress=
+
+
+[/Script/Engine.Engine]
+NearClipPlane=0.03
 
-[/Script/WindowsTargetPlatform.WindowsTargetSettings]
-DefaultGraphicsRHI=DefaultGraphicsRHI_DX12
--D3D12TargetedShaderFormats=PCD3D_SM5
-+D3D12TargetedShaderFormats=PCD3D_SM6
--D3D11TargetedShaderFormats=PCD3D_SM5
-+D3D11TargetedShaderFormats=PCD3D_SM5
