[Template-PhotoStudio] Pointed to the correct default map, updated the launcher image (WIP version), and revised the template description.

diff --git a/Templates/TP_PhotoStudioBP/Config/DefaultEngine.ini b/Templates/TP_PhotoStudioBP/Config/DefaultEngine.ini
--- a/Templates/TP_PhotoStudioBP/Config/DefaultEngine.ini
+++ b/Templates/TP_PhotoStudioBP/Config/DefaultEngine.ini
@@ -84,3 +84,7 @@ ManualIPAddress=
 [/Script/Engine.Engine]
 NearClipPlane=0.03
 
+[/Script/EngineSettings.GameMapsSettings]
+EditorStartupMap=/Game/Studios/Levels/Studio_White.Studio_White
+GameDefaultMap=/Game/Studios/Levels/Studio_White.Studio_White
+
diff --git a/Templates/TP_PhotoStudioBP/Config/TemplateDefs.ini b/Templates/TP_PhotoStudioBP/Config/TemplateDefs.ini
--- a/Templates/TP_PhotoStudioBP/Config/TemplateDefs.ini
+++ b/Templates/TP_PhotoStudioBP/Config/TemplateDefs.ini
@@ -11,7 +11,7 @@ Categories=MFG
 LocalizedDisplayNames=(Language="en", Text="Photo Studio")
 
 ; The localized descriptions for this template, which is displayed when the template is selected.
-LocalizedDescriptions=(Language="en", Text="This template enables you to create product presentations and cinematics in a photographic studio environment.\n\nThis template enables you to use Datasmith to import models from various design applications and easily light them.")
+LocalizedDescriptions=(Language="en", Text="This template enables you to create product presentations and cinematics in a photographic studio environment.\nIt provides two studio variations, white and dark, and is ready to use.")
 
 LocalizedDisplayNames=(Language="ko", Text="포토 스튜디오")
 LocalizedDescriptions=(Language="ko", Text="이 템플릿은 사진 스튜디오 환경에서 제품 시연과 시네마틱을 제작할 수 있도록 해줍니다.\n\n이 템플릿은 다양한 디자인 어플리케이션에서 데이터스미스를 이용한 모델 임포트와 시각화가 가능합니다.")
