[TemplateResources-ConceptCar] Updated the glass part of the mesh with a quad model for improved reflections, Applied material tweaks
[Template-PhotoStudio] Enabled ClearCoat normal for the legacy car paint material in ConceptCar, Updated media for the New Project Window

#virtualized

diff --git a/Templates/TP_PhotoStudioBP/Config/DefaultEngine.ini b/Templates/TP_PhotoStudioBP/Config/DefaultEngine.ini
--- a/Templates/TP_PhotoStudioBP/Config/DefaultEngine.ini
+++ b/Templates/TP_PhotoStudioBP/Config/DefaultEngine.ini
@@ -50,6 +50,7 @@ r.RayTracedTranslucency.UseRayTracedRefraction=1
 r.Lumen.TranslucencyReflections.FrontLayer.EnableForProject=True
 r.Substrate.Debug.AdvancedVisualizationShaders=True
 r.PathTracing.VisibleLights=1
+r.ClearCoatNormal=True
 
 [/Script/Engine.VirtualTexturePoolConfig]
 DefaultSizeInMegabyte=64
diff --git a/Templates/TP_PhotoStudioBP/Config/TemplateDefs.ini b/Templates/TP_PhotoStudioBP/Config/TemplateDefs.ini
--- a/Templates/TP_PhotoStudioBP/Config/TemplateDefs.ini
+++ b/Templates/TP_PhotoStudioBP/Config/TemplateDefs.ini
@@ -89,11 +89,11 @@ ReplacementsInFiles=(Extensions=("cpp","h","ini","cs"), From="%TEMPLATENAME%", T
 ; GraphicsPreset - Choose between Maximum and Scalable graphics.
 ; StarterContent - Choose whether to add starter content to the project.
 ; XR - Choose whether to enable Virtual Reality support for the project.
-;HiddenSettings=All
+HiddenSettings=All
 HiddenSettings=Languages
 HiddenSettings=HardwareTarget
 HiddenSettings=GraphicsPreset 
-;HiddenSettings=StarterContent 
+HiddenSettings=StarterContent 
 HiddenSettings=XR
 
 ; The class types that this project uses.
