[Template-PhotoStudio] Added Automotive Interchange Pipeline preset, Set it as the project default

#virtualized

diff --git a/Templates/TP_PhotoStudioBP/Config/DefaultEngine.ini b/Templates/TP_PhotoStudioBP/Config/DefaultEngine.ini
--- a/Templates/TP_PhotoStudioBP/Config/DefaultEngine.ini
+++ b/Templates/TP_PhotoStudioBP/Config/DefaultEngine.ini
@@ -89,3 +89,11 @@ NearClipPlane=0.03
 EditorStartupMap=/Game/Studios/Levels/Studio_White.Studio_White
 GameDefaultMap=/Game/Studios/Levels/Studio_White.Studio_White
 
+[/Script/InterchangeEngine.InterchangeProjectSettings]
+ContentImportSettings=(DefaultPipelineStackOverride=((Materials, "Materials"),(Textures, "Textures"),(Sounds, "Audio")),ShowImportDialogOverride=((Textures, (bShowImportDialog=False))),PipelineStacks=(("Assets", (Pipelines=("/Game/Studios/InterchangePipeline/InterchangPipeline_Automotive.InterchangPipeline_Automotive","/Interchange/Pipelines/DefaultAssetsPipeline.DefaultAssetsPipeline"),PerTranslatorPipelines=((Translator="/Script/InterchangeImport.InterchangeMaterialXTranslator",Pipelines=("/Interchange/Pipelines/DefaultAssetsPipeline.DefaultAssetsPipeline","/Interchange/Pipelines/DefaultMaterialXPipeline.DefaultMaterialXPipeline")),(Translator="/Script/InterchangeImport.InterchangeGLTFTranslator",Pipelines=("/Interchange/Pipelines/DefaultGLTFAssetsPipeline.DefaultGLTFAssetsPipeline","/Interchange/Pipelines/DefaultGLTFPipeline.DefaultGLTFPipeline"))))),("Materials", (Pipelines=("/Interchange/Pipelines/DefaultMaterialPipeline.DefaultMaterialPipeline"),PerTranslatorPipelines=((Translator="/Script/InterchangeImport.InterchangeMaterialXTranslator",Pipelines=("/Interchange/Pipelines/DefaultMaterialPipeline.DefaultMaterialPipeline","/Interchange/Pipelines/DefaultMaterialXPipeline.DefaultMaterialXPipeline"))))),("Textures", (Pipelines=("/Interchange/Pipelines/DefaultTexturePipeline.DefaultTexturePipeline"))),("Audio", (Pipelines=("/Interchange/Pipelines/DefaultAudioPipeline.DefaultAudioPipeline")))),DefaultPipelineStack="Assets",ImportDialogClass="/Script/InterchangeEditorPipelines.InterchangePipelineConfigurationGeneric",bShowImportDialog=True,bShowReimportDialog=False)
+SceneImportSettings=(PerTranslatorDialogOverride=,PipelineStacks=(("Scene", (Pipelines=("/Interchange/Pipelines/DefaultSceneAssetsPipeline.DefaultSceneAssetsPipeline","/Interchange/Pipelines/DefaultSceneLevelPipeline.DefaultSceneLevelPipeline"),PerTranslatorPipelines=((Translator="/Script/InterchangeImport.InterchangeMaterialXTranslator",Pipelines=("/Interchange/Pipelines/DefaultSceneAssetsPipeline.DefaultSceneAssetsPipeline","/Interchange/Pipelines/DefaultSceneLevelPipeline.DefaultSceneLevelPipeline","/Interchange/Pipelines/DefaultMaterialXPipeline.DefaultMaterialXPipeline")),(Translator="/Script/InterchangeImport.InterchangeGLTFTranslator",Pipelines=("/Interchange/Pipelines/DefaultGLTFSceneAssetsPipeline.DefaultGLTFSceneAssetsPipeline","/Interchange/Pipelines/DefaultSceneLevelPipeline.DefaultSceneLevelPipeline","/Interchange/Pipelines/DefaultGLTFPipeline.DefaultGLTFPipeline")))))),DefaultPipelineStack="Scene",ImportDialogClass="/Script/InterchangeEditorPipelines.InterchangePipelineConfigurationGeneric",bShowImportDialog=True,bShowReimportDialog=False)
+FilePickerClass=/Script/InterchangeEditorUtilities.InterchangeFilePickerGeneric
+bStaticMeshUseSmoothEdgesIfSmoothingInformationIsMissing=True
+GenericPipelineClass=/Script/InterchangePipelines.InterchangeGenericAssetsPipeline
+ConverterDefaultPipeline=/Interchange/Pipelines/DefaultAssetsPipeline.DefaultAssetsPipeline
+
