[Template-PhotoStudio] TemplateDefs.ini: Comment out individual HiddenSettings to prevent warnings during build

diff --git a/Templates/TP_PhotoStudioBP/Config/TemplateDefs.ini b/Templates/TP_PhotoStudioBP/Config/TemplateDefs.ini
--- a/Templates/TP_PhotoStudioBP/Config/TemplateDefs.ini
+++ b/Templates/TP_PhotoStudioBP/Config/TemplateDefs.ini
@@ -90,11 +90,11 @@ ReplacementsInFiles=(Extensions=("cpp","h","ini","cs"), From="%TEMPLATENAME%", T
 ; StarterContent - Choose whether to add starter content to the project.
 ; XR - Choose whether to enable Virtual Reality support for the project.
 HiddenSettings=All
-HiddenSettings=Languages
-HiddenSettings=HardwareTarget
-HiddenSettings=GraphicsPreset 
-HiddenSettings=StarterContent 
-HiddenSettings=XR
+; HiddenSettings=Languages
+; HiddenSettings=HardwareTarget
+; HiddenSettings=GraphicsPreset 
+; HiddenSettings=StarterContent 
+; HiddenSettings=XR  
 
 ; The class types that this project uses.
 ; This is only for display in the template description, and does not affect project creation.
