[Template-PhotoStudio] Added DefaultEditorPerProjectUserSettings.ini to configure the import window to show only essential properties in the project by default.
Organized pipeline presets for Static Meshes, Skeletal Meshes, Audio, Materials & Textures, Textures, and HDR Cubemaps. (WIP Presets)

#virtualized

diff --git a/Templates/TP_PhotoStudioBP/Config/DefaultEditorPerProjectUserSettings.ini b/Templates/TP_PhotoStudioBP/Config/DefaultEditorPerProjectUserSettings.ini
new file mode 100644
--- /dev/null
+++ b/Templates/TP_PhotoStudioBP/Config/DefaultEditorPerProjectUserSettings.ini
@@ -0,0 +1,2 @@
+[InterchangeImportDialogOptions]
+ShowEssentials=True
