[Template-PhotoStudio] Port rendering settings over from SMP-5.7

diff --git a/Templates/TP_PhotoStudioBP/Config/DefaultEngine.ini b/Templates/TP_PhotoStudioBP/Config/DefaultEngine.ini
--- a/Templates/TP_PhotoStudioBP/Config/DefaultEngine.ini
+++ b/Templates/TP_PhotoStudioBP/Config/DefaultEngine.ini
@@ -15,42 +15,52 @@
 
 [/Script/Engine.RendererSettings]
 r.AllowStaticLighting=0
-r.GenerateMeshDistanceFields=True
+r.ClearCoatNormal=True
+r.DefaultFeature.MotionBlur=False
+r.DefaultFeature.AutoExposure=False
+r.DefaultFeature.LocalExposure.ShadowContrastScale=0.8
+r.DefaultFeature.LocalExposure.HighlightContrastScale=0.8
+r.DefaultFeature.AutoExposure.ExtendDefaultLuminanceRange=True
 r.DynamicGlobalIlluminationMethod=1
-r.ReflectionMethod=1
-r.Shadow.Virtual.Enable=1
+r.GBufferFormat=3
+r.GenerateMeshDistanceFields=True
+r.Lumen.TraceMeshSDFs=0
 r.Lumen.LightingDataFormat=1
+r.Lumen.HardwareRayTracing.LightingMode=1
 r.Lumen.Reflections.DenoiserTonemapRange =12
-r.VirtualTextures=True
-r.DefaultFeature.AutoExposure.ExtendDefaultLuminanceRange=True
-r.DefaultFeature.LocalExposure.HighlightContrastScale=0.8
-r.DefaultFeature.LocalExposure.ShadowContrastScale=0.8
-r.SkinCache.CompileShaders=True
-r.RayTracing=True
-r.Substrate=True
-r.VirtualTextures=True
-r.GBufferFormat=3
-r.Substrate.BytesPerPixel=240
-r.Substrate.EnableLayerSupport=True
-r.ReflectionCaptureResolution=1024
-r.RayTracing.Culling=0
-r.SkyAtmosphere.MultiScatteringLUT.HighQuality=1
-r.Translucency.StandardSeparated=1
-r.RayTracedTranslucency.MaxPrimaryHitEvents=3
-r.RayTracedTranslucency.ForceOpaque=0
 r.Lumen.HardwareRayTracing.HitLighting.ForceOpaque=0
-r.OIT.SortedPixels=False
+r.Lumen.TranslucencyReflections.FrontLayer.EnableForProject=True
+r.Lumen.Reflections.HardwareRayTracing.Translucent.MaxRefractionBounces=2
 r.Material.EnergyConservation=True
-r.RayTracing.Shadows.Translucency=1
+r.Mobile.FloatPrecisionMode=2
+r.OIT.SortedPixels=True
+r.PathTracing.VisibleLights=1
+r.RayTracing=True
+r.RayTracing.Culling=0
+r.RayTracing.Shadows=True
 r.RayTracing.NormalBias=0.18
-r.Lumen.HardwareRayTracing.LightingMode=2
-r.Translucent.UsesRectLights=True
+r.RayTracing.Shadows.Translucency=1
 r.RayTracedTranslucency.Allowed=1
+r.RayTracedTranslucency.ForceOpaque=0
+r.RayTracedTranslucency.MaxPrimaryHitEvents=3
 r.RayTracedTranslucency.UseRayTracedRefraction=1
-r.Lumen.TranslucencyReflections.FrontLayer.EnableForProject=True
+r.ReflectionCaptureResolution=1024
+r.ReflectionMethod=1
+r.Shadow.Virtual.Enable=1
+r.SkinCache.CompileShaders=True
+r.SkinCache.SceneMemoryLimitInMB=126.000000
+r.SkyAtmosphere.MultiScatteringLUT.HighQuality=1
+r.Substrate=True
+r.Substrate.BytesPerPixel=240
+r.Substrate.ProjectGBufferFormat=1
+r.Substrate.EnableLayerSupport=True
+r.Substrate.ProjectClosuresPerPixel=5
+r.Substrate.ProjectClosuresPerPixelOverride=1
+r.Substrate.OpaqueMaterialRoughRefraction=True
 r.Substrate.Debug.AdvancedVisualizationShaders=True
-r.PathTracing.VisibleLights=1
-r.ClearCoatNormal=True
+r.Translucency.StandardSeparated=1
+r.Translucent.UsesRectLights=True
+r.VirtualTextures=True
 
 [/Script/Engine.VirtualTexturePoolConfig]
 DefaultSizeInMegabyte=64
@@ -81,7 +91,6 @@
 bUseManualIPAddress=False
 ManualIPAddress=
 
-
 [/Script/Engine.Engine]
 NearClipPlane=0.03
 
@@ -96,4 +105,3 @@
 bStaticMeshUseSmoothEdgesIfSmoothingInformationIsMissing=True
 GenericPipelineClass=/Script/InterchangePipelines.InterchangeGenericAssetsPipeline
 ConverterDefaultPipeline=/Interchange/Pipelines/DefaultAssetsPipeline.DefaultAssetsPipeline
-
