[Template-PhotoStudio/Template-MFG_Blank] Update the DefaultEngine.ini rendering settings: r.RayTracing.Nanite.Mode=1 - streamed out mesh (instead of fallback mesh)

diff --git a/Templates/TP_MFG_BlankBP/Config/DefaultEngine.ini b/Templates/TP_MFG_BlankBP/Config/DefaultEngine.ini
--- a/Templates/TP_MFG_BlankBP/Config/DefaultEngine.ini
+++ b/Templates/TP_MFG_BlankBP/Config/DefaultEngine.ini
@@ -32,6 +32,7 @@ r.RayTracedTranslucency.ForceOpaque=0
 r.RayTracedTranslucency.MaxPrimaryHitEvents=3
 r.RayTracedTranslucency.UseRayTracedRefraction=1
 r.RayTracing.Culling=0
+r.RayTracing.Nanite.Mode=1
 r.RayTracing.NormalBias=0.18
 r.RayTracing.Shadows.Translucency=1
 r.RayTracing.Shadows=True
diff --git a/Templates/TP_PhotoStudioBP/Config/DefaultEngine.ini b/Templates/TP_PhotoStudioBP/Config/DefaultEngine.ini
--- a/Templates/TP_PhotoStudioBP/Config/DefaultEngine.ini
+++ b/Templates/TP_PhotoStudioBP/Config/DefaultEngine.ini
@@ -41,6 +41,7 @@ r.RayTracedTranslucency.ForceOpaque=0
 r.RayTracedTranslucency.MaxPrimaryHitEvents=3
 r.RayTracedTranslucency.UseRayTracedRefraction=1
 r.RayTracing.Culling=0
+r.RayTracing.Nanite.Mode=1
 r.RayTracing.NormalBias=0.18
 r.RayTracing.Shadows.Translucency=1
 r.RayTracing.Shadows=True
