[Template-PhotoStudio/Template-MFG_Blank] Update the DefaultEngine.ini rendering settings: r.RayTracing.ResidentGeometryMemoryPoolSizeInMB=700

diff --git a/Templates/TP_MFG_BlankBP/Config/DefaultEngine.ini b/Templates/TP_MFG_BlankBP/Config/DefaultEngine.ini
--- a/Templates/TP_MFG_BlankBP/Config/DefaultEngine.ini
+++ b/Templates/TP_MFG_BlankBP/Config/DefaultEngine.ini
@@ -34,6 +34,7 @@ r.RayTracedTranslucency.UseRayTracedRefraction=1
 r.RayTracing.Culling=0
 r.RayTracing.Nanite.Mode=1
 r.RayTracing.NormalBias=0.18
+r.RayTracing.ResidentGeometryMemoryPoolSizeInMB=700
 r.RayTracing.Shadows.Translucency=1
 r.RayTracing.Shadows=True
 r.RayTracing=True
diff --git a/Templates/TP_PhotoStudioBP/Config/DefaultEngine.ini b/Templates/TP_PhotoStudioBP/Config/DefaultEngine.ini
--- a/Templates/TP_PhotoStudioBP/Config/DefaultEngine.ini
+++ b/Templates/TP_PhotoStudioBP/Config/DefaultEngine.ini
@@ -43,6 +43,7 @@ r.RayTracedTranslucency.UseRayTracedRefraction=1
 r.RayTracing.Culling=0
 r.RayTracing.Nanite.Mode=1
 r.RayTracing.NormalBias=0.18
+r.RayTracing.ResidentGeometryMemoryPoolSizeInMB=700
 r.RayTracing.Shadows.Translucency=1
 r.RayTracing.Shadows=True
 r.RayTracing=True
