[Template-PhotoStudio] Merge conflict between Main and Release-5.7

diff --git a/Templates/TP_PhotoStudioBP/Config/DefaultEngine.ini b/Templates/TP_PhotoStudioBP/Config/DefaultEngine.ini
--- a/Templates/TP_PhotoStudioBP/Config/DefaultEngine.ini
+++ b/Templates/TP_PhotoStudioBP/Config/DefaultEngine.ini
@@ -1,9 +1,18 @@
+[Audio]
+UseAudioMixer=True
+
 [/Script/HardwareTargeting.HardwareTargetingSettings]
 TargetedHardwareClass=Desktop
 AppliedTargetedHardwareClass=Desktop
 DefaultGraphicsPerformance=Maximum
 AppliedDefaultGraphicsPerformance=Maximum
 
+[/Script/WindowsTargetPlatform.WindowsTargetSettings]
+DefaultGraphicsRHI=DefaultGraphicsRHI_DX12
+DefaultGraphicsRHI=DefaultGraphicsRHI_DX12
+!D3D12TargetedShaderFormats=ClearArray
++D3D12TargetedShaderFormats=PCD3D_SM6
+
 [/Script/Engine.RendererSettings]
 r.AllowStaticLighting=0
 r.GenerateMeshDistanceFields=True
@@ -49,29 +58,18 @@ bPoolAutoGrowInEditor=False
 +Pools=(Formats=(PF_BC5),MinTileSize=136,MaxTileSize=136,SizeInMegabyte=84,bEnableResidencyMipMapBias=False,bAllowSizeScale=False,MinScaledSizeInMegabyte=0,MaxScaledSizeInMegabyte=0)
 
 
-[/Script/EngineSettings.GameMapsSettings]
-;defines what level to be loaded by default in the editor and at playback time
-GameDefaultMap=/Game/Studios/Levels/Studios.Studios
-EditorStartupMap=/Game/Studios/Levels/Studios.Studios
-GlobalDefaultGameMode=/Game/PhotoStudio/Blueprints/BP_PhotoStageGameMode.BP_PhotoStageGameMode_C
+[/Script/WorldPartitionEditor.WorldPartitionEditorSettings]
+CommandletClass=Class'/Script/UnrealEd.WorldPartitionConvertCommandlet'
 
-[/Script/UnrealEd.UnrealEdEngine]
-; Remove Engine Template maps
--TemplateMapInfos=(ThumbnailTexture=Texture2D'/Engine/Maps/Templates/Thumbnails/Default.Default',Map="/Engine/Maps/Templates/Template_Default")
-+TemplateMapInfos=(ThumbnailTexture=Texture2D'/Engine/Maps/Templates/Thumbnails/TimeOfDay.TimeOfDay',Map="/Engine/Maps/Templates/TimeOfDay_Default")
--TemplateMapInfos=(ThumbnailTexture=Texture2D'/Engine/Maps/Templates/Thumbnails/VR-Basic.VR-Basic',Map="/Engine/Maps/Templates/VR-Basic")
-
-[/Script/WindowsTargetPlatform.WindowsTargetSettings]
-DefaultGraphicsRHI=DefaultGraphicsRHI_DX12
--D3D12TargetedShaderFormats=PCD3D_SM5
-+D3D12TargetedShaderFormats=PCD3D_SM6
--D3D11TargetedShaderFormats=PCD3D_SM5
-+D3D11TargetedShaderFormats=PCD3D_SM5
+[/Script/Engine.UserInterfaceSettings]
+bAuthorizeAutomaticWidgetVariableCreation=False
+FontDPIPreset=Standard
+FontDPI=72
 
 [/Script/AndroidFileServerEditor.AndroidFileServerRuntimeSettings]
 bEnablePlugin=True
 bAllowNetworkConnection=True
-SecurityToken=63A034444FF73356F37D878120E802C8
+SecurityToken=79BFBCEF42D211CDBF800C9F2597C868
 bIncludeInShipping=False
 bAllowExternalStartInShipping=False
 bCompileAFSProject=False
@@ -82,6 +80,7 @@ ConnectionType=USBOnly
 bUseManualIPAddress=False
 ManualIPAddress=
 
+
 [/Script/Engine.Engine]
 NearClipPlane=0.03
 
