[Template-PhotoStudio/MFG_Blank] Increase r.EyeAdaptation.CachedLightingPreExposure from 4 to 8 to support a wider exposure range and prevent cached lighting clipping

#virtualized
#Synced-CL 51566404

#Synced-CL 51566404

diff --git a/Templates/TP_MFG_BlankBP/Config/DefaultEngine.ini b/Templates/TP_MFG_BlankBP/Config/DefaultEngine.ini
--- a/Templates/TP_MFG_BlankBP/Config/DefaultEngine.ini
+++ b/Templates/TP_MFG_BlankBP/Config/DefaultEngine.ini
@@ -14,6 +14,7 @@ r.DefaultFeature.LocalExposure.HighlightContrastScale=0.8
 r.DefaultFeature.LocalExposure.ShadowContrastScale=0.8
 r.DefaultFeature.MotionBlur=False
 r.DynamicGlobalIlluminationMethod=1
+r.EyeAdaptation.CachedLightingPreExposure=8
 r.GBufferFormat=3
 r.GenerateMeshDistanceFields=True
 r.Lumen.HardwareRayTracing.HitLighting.ForceOpaque=0
diff --git a/Templates/TP_PhotoStudioBP/Config/DefaultEngine.ini b/Templates/TP_PhotoStudioBP/Config/DefaultEngine.ini
--- a/Templates/TP_PhotoStudioBP/Config/DefaultEngine.ini
+++ b/Templates/TP_PhotoStudioBP/Config/DefaultEngine.ini
@@ -23,6 +23,7 @@ r.DefaultFeature.LocalExposure.HighlightContrastScale=0.8
 r.DefaultFeature.LocalExposure.ShadowContrastScale=0.8
 r.DefaultFeature.MotionBlur=False
 r.DynamicGlobalIlluminationMethod=1
+r.EyeAdaptation.CachedLightingPreExposure=8
 r.GBufferFormat=3
 r.GenerateMeshDistanceFields=True
 r.Lumen.HardwareRayTracing.HitLighting.ForceOpaque=0
