[Template-PhotoStudio/MFG_Blank] Update ini to avoid Scalability warnings - Removed r.MotionBlurQuality and r.BloomQuality, and moved r.LumenScene.SurfaceCache.CardMinResolution and r.Lumen.ScreenProbeGather.DownsampleFactor to DefaultScalability.ini (Epic/Cine only)

#Synced-CL 51571364

diff --git a/Templates/TP_MFG_BlankBP/Config/DefaultEngine.ini b/Templates/TP_MFG_BlankBP/Config/DefaultEngine.ini
--- a/Templates/TP_MFG_BlankBP/Config/DefaultEngine.ini
+++ b/Templates/TP_MFG_BlankBP/Config/DefaultEngine.ini
@@ -7,7 +7,6 @@ AppliedDefaultGraphicsPerformance=Maximum
 [/Script/Engine.RendererSettings]
 r.AllowStaticLighting=0
 r.AntiAliasingMethod=4
-r.BloomQuality=2
 r.DefaultFeature.AutoExposure.ExtendDefaultLuminanceRange=True
 r.DefaultFeature.AutoExposure=False
 r.DefaultFeature.LocalExposure.HighlightContrastScale=0.8
@@ -23,7 +22,6 @@ r.Lumen.LightingDataFormat=1
 r.Lumen.Reflections.BilateralFilter.KernelRadius=32
 r.Lumen.Reflections.DenoiserTonemapRange =12
 r.Lumen.Reflections.HardwareRayTracing.Translucent.MaxRefractionBounces=2
-r.Lumen.ScreenProbeGather.DownsampleFactor=4
 r.Lumen.ScreenProbeGather.ShortRangeAO.HorizonSearch.ForegroundSampleRejectDistanceFraction=0.025
 r.Lumen.ScreenProbeGather.Temporal.DistanceThreshold=0.001
 r.Lumen.ScreenProbeGather.Temporal.NormalThreshold=5
@@ -31,11 +29,9 @@ r.Lumen.ScreenProbeGather.Temporal.RejectBasedOnNormal=1
 r.Lumen.ScreenProbeGather.TemporalFilterProbes=1
 r.Lumen.TraceMeshSDFs=0
 r.Lumen.TranslucencyReflections.FrontLayer.EnableForProject=True
-r.LumenScene.SurfaceCache.CardMinResolution=2
 r.LumenScene.SurfaceCache.MeshCardsMinSize=4
 r.Material.EnergyConservation=True
 r.Mobile.FloatPrecisionMode=2
-r.MotionBlurQuality=0
 r.OIT.SortedPixels=True
 r.PathTracing.Denoiser.Prepass.RankedLuminanceVariance=0
 r.PathTracing.VisibleLights=1
diff --git a/Templates/TP_MFG_BlankBP/Config/DefaultScalability.ini b/Templates/TP_MFG_BlankBP/Config/DefaultScalability.ini
--- a/Templates/TP_MFG_BlankBP/Config/DefaultScalability.ini
+++ b/Templates/TP_MFG_BlankBP/Config/DefaultScalability.ini
@@ -17,4 +17,12 @@ r.Shadow.Virtual.NonNanite.IncludeInCoarsePages=0
 r.Shadow.Virtual.NonNanite.IncludeInCoarsePages=1
 
 [TextureQuality@3]
-r.Streaming.PoolSize=2000
\ No newline at end of file
+r.Streaming.PoolSize=2000
+
+[GlobalIlluminationQuality@3]
+r.Lumen.ScreenProbeGather.DownsampleFactor=4
+r.LumenScene.SurfaceCache.CardMinResolution=2
+
+[GlobalIlluminationQuality@Cine]
+r.Lumen.ScreenProbeGather.DownsampleFactor=4
+r.LumenScene.SurfaceCache.CardMinResolution=2
\ No newline at end of file
diff --git a/Templates/TP_PhotoStudioBP/Config/DefaultEngine.ini b/Templates/TP_PhotoStudioBP/Config/DefaultEngine.ini
--- a/Templates/TP_PhotoStudioBP/Config/DefaultEngine.ini
+++ b/Templates/TP_PhotoStudioBP/Config/DefaultEngine.ini
@@ -16,7 +16,6 @@ DefaultGraphicsRHI=DefaultGraphicsRHI_DX12
 [/Script/Engine.RendererSettings]
 r.AllowStaticLighting=0
 r.AntiAliasingMethod=4
-r.BloomQuality=2
 r.DefaultFeature.AutoExposure.ExtendDefaultLuminanceRange=True
 r.DefaultFeature.AutoExposure=False
 r.DefaultFeature.LocalExposure.HighlightContrastScale=0.8
@@ -32,7 +31,6 @@ r.Lumen.LightingDataFormat=1
 r.Lumen.Reflections.BilateralFilter.KernelRadius=32
 r.Lumen.Reflections.DenoiserTonemapRange =12
 r.Lumen.Reflections.HardwareRayTracing.Translucent.MaxRefractionBounces=2
-r.Lumen.ScreenProbeGather.DownsampleFactor=4
 r.Lumen.ScreenProbeGather.ShortRangeAO.HorizonSearch.ForegroundSampleRejectDistanceFraction=0.025
 r.Lumen.ScreenProbeGather.Temporal.DistanceThreshold=0.001
 r.Lumen.ScreenProbeGather.Temporal.NormalThreshold=5
@@ -40,11 +38,9 @@ r.Lumen.ScreenProbeGather.Temporal.RejectBasedOnNormal=1
 r.Lumen.ScreenProbeGather.TemporalFilterProbes=1
 r.Lumen.TraceMeshSDFs=0
 r.Lumen.TranslucencyReflections.FrontLayer.EnableForProject=True
-r.LumenScene.SurfaceCache.CardMinResolution=2
 r.LumenScene.SurfaceCache.MeshCardsMinSize=4
 r.Material.EnergyConservation=True
 r.Mobile.FloatPrecisionMode=2
-r.MotionBlurQuality=0
 r.OIT.SortedPixels=True
 r.PathTracing.Denoiser.Prepass.RankedLuminanceVariance=0
 r.PathTracing.VisibleLights=1
diff --git a/Templates/TP_PhotoStudioBP/Config/DefaultScalability.ini b/Templates/TP_PhotoStudioBP/Config/DefaultScalability.ini
new file mode 100644
--- /dev/null
+++ b/Templates/TP_PhotoStudioBP/Config/DefaultScalability.ini
@@ -0,0 +1,7 @@
+[GlobalIlluminationQuality@3]
+r.Lumen.ScreenProbeGather.DownsampleFactor=4
+r.LumenScene.SurfaceCache.CardMinResolution=2
+
+[GlobalIlluminationQuality@Cine]
+r.Lumen.ScreenProbeGather.DownsampleFactor=4
+r.LumenScene.SurfaceCache.CardMinResolution=2
\ No newline at end of file
